Fix crash in ConcertInsightsClient due to adding status bar widget even though slate application is not initialized (e.g. running commandlet).
Fixes blocker in Build Engine Localization build farm.

diff --git a/Engine/Plugins/Developer/Concert/ConcertInsights/ConcertInsightsClient/Source/ConcertInsightsClient/Private/ConcertInsightsClient.cpp b/Engine/Plugins/Developer/Concert/ConcertInsights/ConcertInsightsClient/Source/ConcertInsightsClient/Private/ConcertInsightsClient.cpp
--- a/Engine/Plugins/Developer/Concert/ConcertInsights/ConcertInsightsClient/Source/ConcertInsightsClient/Private/ConcertInsightsClient.cpp
+++ b/Engine/Plugins/Developer/Concert/ConcertInsights/ConcertInsightsClient/Source/ConcertInsightsClient/Private/ConcertInsightsClient.cpp
@@ -5,6 +5,7 @@
 #include "StatusBar/SMultiUserStatusBar.h"
 #include "StatusBar/StatusBarExtension.h"
 
+#include "Application/SlateApplicationBase.h"
 #include "Misc/CoreDelegates.h"
 
 #define LOCTEXT_NAMESPACE "FConcertInsightsEditorModule"
@@ -21,11 +22,14 @@ namespace UE::ConcertInsightsClient
 
 	void FConcertInsightsClientModule::PostEngineInit()
 	{
-		// Multi User does not add any status bar by itself. For now, we'll add one but in the future we may want this to be in the MU module.
-		ExtendEditorStatusBarWithMultiUserWidget();
+		if (FSlateApplicationBase::IsInitialized())
+		{
+			// Multi User does not add any status bar by itself. For now, we'll add one but in the future we may want this to be in the MU module.
+			ExtendEditorStatusBarWithMultiUserWidget();
 		
-		TraceControls = ConcertInsightsCore::FTraceControls::Make<FClientTraceControls>();
-		ExtendMultiUserStatusBarWithInsights(*TraceControls);
+			TraceControls = ConcertInsightsCore::FTraceControls::Make<FClientTraceControls>();
+			ExtendMultiUserStatusBarWithInsights(*TraceControls);
+		}
 	}
 }
 
