Fix duplicate loca keys

diff --git a/Engine/Plugins/Developer/Concert/ConcertInsights/ConcertInsightsClient/Source/ConcertInsightsClient/Private/StatusBar/StatusBarExtension.cpp b/Engine/Plugins/Developer/Concert/ConcertInsights/ConcertInsightsClient/Source/ConcertInsightsClient/Private/StatusBar/StatusBarExtension.cpp
--- a/Engine/Plugins/Developer/Concert/ConcertInsights/ConcertInsightsClient/Source/ConcertInsightsClient/Private/StatusBar/StatusBarExtension.cpp
+++ b/Engine/Plugins/Developer/Concert/ConcertInsights/ConcertInsightsClient/Source/ConcertInsightsClient/Private/StatusBar/StatusBarExtension.cpp
@@ -33,8 +33,8 @@ namespace UE::ConcertInsightsClient
 				TAttribute<FText>::CreateLambda([&Controls]()
 				{
 					return Controls.IsTracing()
-						? LOCTEXT("Menu.Tracing.ToggleSynchronizedTrace.Label", "Stop synchronized trace")
-						: LOCTEXT("Menu.Tracing.ToggleSynchronizedTrace.Label", "Start synchronized trace");
+						? LOCTEXT("Menu.Tracing.ToggleSynchronizedTrace.StopLabel", "Stop synchronized trace")
+						: LOCTEXT("Menu.Tracing.ToggleSynchronizedTrace.StartLabel", "Start synchronized trace");
 				}),
 				TAttribute<FText>::CreateLambda([&Controls]()
 				{
diff --git a/Engine/Plugins/Developer/Concert/ConcertInsights/ConcertInsightsCore/Source/ConcertInsightsCore/Private/TraceControls.cpp b/Engine/Plugins/Developer/Concert/ConcertInsights/ConcertInsightsCore/Source/ConcertInsightsCore/Private/TraceControls.cpp
--- a/Engine/Plugins/Developer/Concert/ConcertInsights/ConcertInsightsCore/Source/ConcertInsightsCore/Private/TraceControls.cpp
+++ b/Engine/Plugins/Developer/Concert/ConcertInsights/ConcertInsightsCore/Source/ConcertInsightsCore/Private/TraceControls.cpp
@@ -40,7 +40,7 @@ namespace UE::ConcertInsightsCore
 
 		if (!LocallyStartSynchronizedTrace({ Session }, Args))
 		{
-			SET_REASON(FailReason, LOCTEXT("Reason.AlreadyRunning", "Failed to start local trace. See log."));
+			SET_REASON(FailReason, LOCTEXT("Reason.FailedToStart", "Failed to start local trace. See log."));
 			return false;
 		}
 		
